Merging using UE4-To-UE4-LauncherDev

diff --git a/Engine/Source/Developer/Toolbox/Private/Toolbox.cpp b/Engine/Source/Developer/Toolbox/Private/Toolbox.cpp
new file mode 100644
--- /dev/null
+++ b/Engine/Source/Developer/Toolbox/Private/Toolbox.cpp
@@ -0,0 +1,192 @@
+// Copyright 1998-2015 Epic Games, Inc. All Rights Reserved.
+
+#include "Toolbox.h"
+#include "ModuleManager.h"
+#include "GammaUI.h"
+#include "ModuleUIInterface.h"
+#include "Editor/WorkspaceMenuStructure/Public/WorkspaceMenuStructureModule.h"
+#include "MainFrame.h"
+#include "DesktopPlatformModule.h"
+#include "SDockTab.h"
+#include "SNotificationList.h"
+#include "NotificationManager.h"
+#include "STestSuite.h"
+#include "ISlateReflectorModule.h"
+
+IMPLEMENT_MODULE( FToolboxModule, Toolbox );
+
+
+static bool CanShowModulesTab()
+{
+	FString SolutionPath;
+	return FDesktopPlatformModule::Get()->GetSolutionPath(SolutionPath);
+}
+
+class SDebugPanel : public SCompoundWidget
+{
+public:
+
+	SLATE_BEGIN_ARGS(SDebugPanel){}
+	SLATE_END_ARGS()
+
+	BEGIN_SLATE_FUNCTION_BUILD_OPTIMIZATION
+	void Construct( const FArguments& )
+	{
+		ChildSlot
+		[
+			SNew( SVerticalBox )
+			+ SVerticalBox::Slot()
+			.AutoHeight()
+			.Padding( 4.0f )
+			.HAlign(HAlign_Left)
+			[
+				SNew( SButton )
+				.Text( NSLOCTEXT("DeveloperToolbox", "ReloadTextures", "Reload Textures") )
+				.OnClicked( this, &SDebugPanel::OnReloadTexturesClicked )
+			]
+			
+			+ SVerticalBox::Slot()
+			.AutoHeight()
+			.Padding( 4.0f )
+			.HAlign(HAlign_Left)
+			[
+				SNew( SButton )
+				.Text( NSLOCTEXT("DeveloperToolbox", "FlushFontCache", "Flush Font Cache") )
+				.OnClicked( this, &SDebugPanel::OnFlushFontCacheClicked )
+			]
+			+ SVerticalBox::Slot()
+			.AutoHeight()
+			.Padding( 4.0f )
+			.HAlign(HAlign_Left)
+			[
+				SNew( SButton )
+				.Text( NSLOCTEXT("DeveloperToolbox", "TestSuite", "Test Suite") )
+				.OnClicked( this, &SDebugPanel::OnTestSuiteClicked )
+			]
+			+ SVerticalBox::Slot()
+			.AutoHeight()
+			.Padding( 4.0f )
+			.HAlign(HAlign_Left)
+			[
+				SNew( SButton )
+				.Text( NSLOCTEXT("DeveloperToolbox", "DisplayTextureAtlases", "Display Texture Atlases") )
+				.OnClicked( this, &SDebugPanel::OnDisplayTextureAtlases )
+			]
+			+ SVerticalBox::Slot()
+			.AutoHeight()
+			.Padding( 4.0f )
+			.HAlign(HAlign_Left)
+			[
+				SNew( SButton )
+				.Text( NSLOCTEXT("DeveloperToolbox", "DisplayFontAtlases", "Display Font Atlases") )
+				.OnClicked( this, &SDebugPanel::OnDisplayFontAtlases )
+			]
+		];
+	}
+	END_SLATE_FUNCTION_BUILD_OPTIMIZATION
+
+	
+	FReply OnReloadTexturesClicked()
+	{
+		FSlateApplication::Get().GetRenderer()->ReloadTextureResources();
+		return FReply::Handled();
+	}
+
+	FReply OnDisplayTextureAtlases()
+	{
+		static const FName SlateReflectorModuleName("SlateReflector");
+		FModuleManager::LoadModuleChecked<ISlateReflectorModule>(SlateReflectorModuleName).DisplayTextureAtlasVisualizer();
+		return FReply::Handled();
+	}
+
+	FReply OnDisplayFontAtlases()
+	{
+		static const FName SlateReflectorModuleName("SlateReflector");
+		FModuleManager::LoadModuleChecked<ISlateReflectorModule>(SlateReflectorModuleName).DisplayFontAtlasVisualizer();
+		return FReply::Handled();
+	}
+
+	FReply OnFlushFontCacheClicked()
+	{
+		FSlateApplication::Get().GetRenderer()->FlushFontCache();
+		return FReply::Handled();
+	}
+
+	FReply OnTestSuiteClicked()
+	{
+		RestoreSlateTestSuite();
+		return FReply::Handled();
+	}
+};
+
+TSharedRef<SDockTab> CreateDebugToolsTab( const FSpawnTabArgs& Args )
+{
+	FGammaUI& GammaUI = FModuleManager::Get().LoadModuleChecked<FGammaUI>( FName("GammaUI") );
+
+	return
+		SNew(SDockTab)
+		.TabRole( ETabRole::NomadTab )
+		[
+			SNew(SVerticalBox)
+			+SVerticalBox::Slot()
+			.AutoHeight()
+			[
+				SNew( SDebugPanel )
+			]
+			+SVerticalBox::Slot()
+			.AutoHeight()
+			[
+				GammaUI.GetGammaUIPanel().ToSharedRef()
+			]
+		];
+}
+
+TSharedRef<SDockTab> CreateModulesTab( const FSpawnTabArgs& Args )
+{
+	IModuleUIInterface& ModuleUI = FModuleManager::Get().LoadModuleChecked<IModuleUIInterface>( FName("ModuleUI") );
+
+	return
+		SNew(SDockTab)
+		.TabRole( ETabRole::NomadTab )
+		[
+			SNew(SBorder)
+			.Padding(2)
+			.BorderImage( FEditorStyle::GetBrush("ToolPanel.GroupBorder") )
+			[
+				ModuleUI.GetModuleUIWidget()
+			]
+		];
+}
+
+
+
+void FToolboxModule::StartupModule()
+{
+	const IWorkspaceMenuStructure& MenuStructure = WorkspaceMenu::GetMenuStructure();
+
+	FGlobalTabmanager::Get()->RegisterNomadTabSpawner( "DebugTools", FOnSpawnTab::CreateStatic( &CreateDebugToolsTab ) )
+		.SetDisplayName( NSLOCTEXT("Toolbox", "DebugTools", "Debug Tools") )
+		.SetTooltipText( NSLOCTEXT("Toolbox", "DebugToolsTooltipText", "Open the Debug Tools tab.") )
+		.SetGroup( MenuStructure.GetDeveloperToolsDebugCategory() )
+		.SetIcon( FSlateIcon(FEditorStyle::GetStyleSetName(), "DebugTools.TabIcon") );
+	if ( CanShowModulesTab() )
+	{
+		FGlobalTabmanager::Get()->RegisterNomadTabSpawner( "ModulesTab", FOnSpawnTab::CreateStatic( &CreateModulesTab ) )
+			.SetDisplayName( NSLOCTEXT("Toolbox", "Modules", "Modules") )
+			.SetTooltipText( NSLOCTEXT("Toolbox", "ModulesTooltipText", "Open the Modules tab.") )
+			.SetGroup( MenuStructure.GetDeveloperToolsMiscCategory() )
+			.SetIcon( FSlateIcon(FEditorStyle::GetStyleSetName(), "Modules.TabIcon") );
+	}
+}
+
+void FToolboxModule::ShutdownModule()
+{
+	if (FSlateApplication::IsInitialized())
+	{
+		FGlobalTabmanager::Get()->UnregisterNomadTabSpawner("DebugTools");
+		if (CanShowModulesTab())
+		{
+			FGlobalTabmanager::Get()->UnregisterNomadTabSpawner("ModulesTab");
+		}
+	}
+}
diff --git a/Engine/Source/Developer/Toolbox/Public/Toolbox.h b/Engine/Source/Developer/Toolbox/Public/Toolbox.h
new file mode 100644
--- /dev/null
+++ b/Engine/Source/Developer/Toolbox/Public/Toolbox.h
@@ -0,0 +1,19 @@
+// Copyright 1998-2015 Epic Games, Inc. All Rights Reserved.
+
+#pragma once
+
+#include "SlateBasics.h"
+#include "EditorStyle.h"
+#include "ModuleInterface.h"
+
+
+/**
+ * Slate toolbox module public interface
+ */
+class FToolboxModule : public IModuleInterface
+{
+
+public:
+	virtual void StartupModule();
+	virtual void ShutdownModule();
+};
diff --git a/Engine/Source/Developer/Toolbox/Toolbox.Build.cs b/Engine/Source/Developer/Toolbox/Toolbox.Build.cs
new file mode 100644
--- /dev/null
+++ b/Engine/Source/Developer/Toolbox/Toolbox.Build.cs
@@ -0,0 +1,38 @@
+// Copyright 1998-2015 Epic Games, Inc. All Rights Reserved.
+
+using UnrealBuildTool;
+
+public class Toolbox : ModuleRules
+{
+	public Toolbox(TargetInfo Target)
+	{
+		PrivateIncludePathModuleNames.AddRange(
+			new string[] {
+				"GammaUI",
+				"MainFrame",
+				"ModuleUI"
+			}
+		);
+
+		PrivateDependencyModuleNames.AddRange(
+			new string[] {
+				"Core",
+                "InputCore",
+				"Slate",
+				"SlateCore",
+				"SlateReflector",
+                "EditorStyle",
+				"DesktopPlatform",
+				"AppFramework"
+			}
+		);
+
+		DynamicallyLoadedModuleNames.AddRange(
+			new string[] {
+				"GammaUI", 
+				"MainFrame",
+				"ModuleUI"
+			}
+		);
+	}
+}
\ No newline at end of file
